Fix module name in argument to IMPLEMENT_MODULE macro; causing link errors when compiling monolithic editor.

diff --git a/Engine/Plugins/Developer/ActorLayerUtilities/Source/ActorLayerUtilitiesEditor/Private/ActorLayerUtilitiesEditor.cpp b/Engine/Plugins/Developer/ActorLayerUtilities/Source/ActorLayerUtilitiesEditor/Private/ActorLayerUtilitiesEditor.cpp
--- a/Engine/Plugins/Developer/ActorLayerUtilities/Source/ActorLayerUtilitiesEditor/Private/ActorLayerUtilitiesEditor.cpp
+++ b/Engine/Plugins/Developer/ActorLayerUtilities/Source/ActorLayerUtilitiesEditor/Private/ActorLayerUtilitiesEditor.cpp
@@ -36,4 +36,4 @@ class FActorLayerUtilitiesEditorModule : public IModuleInterface
 };
 
 
-IMPLEMENT_MODULE(FActorLayerUtilitiesEditorModule, ActorLayerUtilities)
\ No newline at end of file
+IMPLEMENT_MODULE(FActorLayerUtilitiesEditorModule, ActorLayerUtilitiesEditor)
\ No newline at end of file
