Fix allows ActorLayerUtilities in cooked Composure project

diff --git a/Engine/Plugins/Developer/ActorLayerUtilities/Source/ActorLayerUtilities/Public/ActorLayerUtilities.h b/Engine/Plugins/Developer/ActorLayerUtilities/Source/ActorLayerUtilities/Public/ActorLayerUtilities.h
--- a/Engine/Plugins/Developer/ActorLayerUtilities/Source/ActorLayerUtilities/Public/ActorLayerUtilities.h
+++ b/Engine/Plugins/Developer/ActorLayerUtilities/Source/ActorLayerUtilities/Public/ActorLayerUtilities.h
@@ -31,6 +31,6 @@ public:
 	/**
 	 * Get all the actors in this layer
 	 */
-	UFUNCTION(BlueprintCallable, Category=Layer, meta=(WorldContext=WorldContextObject))
+	UFUNCTION(BlueprintCallable, Category=Layer, meta=(WorldContext=WorldContextObject,DevelopmentOnly))
 	static TArray<AActor*> GetActors(UObject* WorldContextObject, const FActorLayer& ActorLayer);
 };
\ No newline at end of file
